Output dependency lists to temporary files before moving into place, so that incomplete/cancelled builds don't cause corrupt output files. Also safeguard against reading corrupt files in UBT.

#rnx

diff --git a/Engine/Extras/cl-filter/cl-filter.cpp b/Engine/Extras/cl-filter/cl-filter.cpp
--- a/Engine/Extras/cl-filter/cl-filter.cpp
+++ b/Engine/Extras/cl-filter/cl-filter.cpp
@@ -79,9 +79,16 @@ int wmain(int ArgC, const wchar_t* ArgV[])
 	CloseHandle(StdOutWriteHandle);
 	CloseHandle(StdErrWriteHandle);
 
+	// Delete the output file
+	DeleteFileW(OutputFileName);
+
+	// Get the path to a temporary output filename
+	std::wstring TempOutputFileName(OutputFileName);
+	TempOutputFileName += L".tmp";
+
 	// Create a file to contain the dependency list
 	FILE* OutputFile;
-	if (_wfopen_s(&OutputFile, OutputFileName, L"wt") != 0)
+	if (_wfopen_s(&OutputFile, TempOutputFileName.c_str(), L"wt") != 0)
 	{
 		wprintf(L"ERROR: Unable to open %s for output", OutputFileName);
 		return -1;
@@ -199,5 +206,11 @@ int wmain(int ArgC, const wchar_t* ArgV[])
 
 	fclose(OutputFile);
 
+	if (ExitCode == 0 && !MoveFile(TempOutputFileName.c_str(), OutputFileName))
+	{
+		wprintf(L"ERROR: Unable to rename %s to %s\n", TempOutputFileName.c_str(), OutputFileName);
+		ExitCode = 1;
+	}
+
 	return ExitCode;
 }
diff --git a/Engine/Source/Programs/UnrealBuildTool/System/CppDependencyCache.cs b/Engine/Source/Programs/UnrealBuildTool/System/CppDependencyCache.cs
--- a/Engine/Source/Programs/UnrealBuildTool/System/CppDependencyCache.cs
+++ b/Engine/Source/Programs/UnrealBuildTool/System/CppDependencyCache.cs
@@ -113,7 +113,16 @@ namespace UnrealBuildTool
 				return false;
 			}
 
-			return TryGetDependenciesInternal(InputFile, out OutDependencyItems);
+			try
+			{
+				return TryGetDependenciesInternal(InputFile, out OutDependencyItems);
+			}
+			catch(Exception Ex)
+			{
+				Log.TraceLog("Unable to read {0}:\n{1}", InputFile, ExceptionUtils.FormatExceptionDetails(Ex));
+				OutDependencyItems = null;
+				return false;
+			}
 		}
 
 		/// <summary>
