Changing how cl-filter.exe takes arguments and cleaning up some of the parse/aggregate timing code in VCToolChain.cs.

diff --git a/Engine/Extras/cl-filter/cl-filter.cpp b/Engine/Extras/cl-filter/cl-filter.cpp
--- a/Engine/Extras/cl-filter/cl-filter.cpp
+++ b/Engine/Extras/cl-filter/cl-filter.cpp
@@ -13,35 +13,91 @@
 
 void GetLocalizedIncludePrefixes(const wchar_t* CompilerPath, std::vector<std::vector<char>>& Prefixes);
 
-int wmain(int ArgC, const wchar_t* ArgV[])
+#define EXIT_CODE_SUCCESS 0
+#define EXIT_CODE_FAILURE -1
+
+#define OUTPUT_FILENAME_INDEX 1
+#define TIMING_FILENAME_INDEX 2
+#define COMPILER_FILENAME_INDEX_OFFSET 2
+#define MINIMUM_VALID_COMMAND_LINE_ARGS 4
+
+void PrintUsage()
 {
-	// Get the full command line, and find the '--' marker
-	wchar_t* CmdLine = ::GetCommandLineW();
+	wprintf(L"Usage: cl-filter.exe -dependencies=<dependencies-file> -compiler=<compiler-file> [Optional Args] -- <child command line>\n");
+	wprintf(L"Optional Args:\n");
+	wprintf(L"\t-timing=<timing-file>\tThe file to write any timing data to.\n");
+}
 
-	wchar_t *ChildCmdLine = wcsstr(CmdLine, L" -- ");
-	if (ChildCmdLine == nullptr)
+struct CommandLineArgs
+{
+	const wchar_t*	OutputFileName = nullptr;
+	const wchar_t*	TimingFileName = nullptr;
+	const wchar_t*	CompilerFileName = nullptr;
+};
+
+CommandLineArgs* ParseCommandLineArgs(int ArgC, const wchar_t* ArgV[])
+{
+	CommandLineArgs* ParsedArgs = new CommandLineArgs();
+	for (int i = 1; i < ArgC; ++i)
 	{
-		wprintf(L"ERROR: Unable to find child command line (%s)", CmdLine);
-		return -1;
+		// If this arg is the splitter, we're done parsing args.
+		if (wcscmp(ArgV[i], L"--") == 0)
+		{
+			break;
+		}
+
+		// If it starts with a -, parse the arg.
+		if (ArgV[i][0] == '-')
+		{
+			if (wcsstr(ArgV[i], L"-dependencies=") != nullptr)
+			{
+				ParsedArgs->OutputFileName = ArgV[i] + 14;
+			}
+			else if (wcsstr(ArgV[i], L"-compiler=") != nullptr)
+			{
+				ParsedArgs->CompilerFileName = ArgV[i] + 10;
+			}
+			else if (wcsstr(ArgV[i], L"-timing=") != nullptr)
+			{
+				ParsedArgs->TimingFileName = ArgV[i] + 8;
+			}
+		}
 	}
-	ChildCmdLine += 4;
 
-	// Make sure we've got an output file and compiler path
-	if (ArgC <= 4 || (wcscmp(ArgV[2], L"--") != 0 && wcscmp(ArgV[3], L"--") != 0))
+	return ParsedArgs;
+}
+
+int wmain(int ArgC, const wchar_t* ArgV[])
+{
+	wchar_t* CmdLine = ::GetCommandLineW();
+	CommandLineArgs* ParsedArgs = ParseCommandLineArgs(ArgC, ArgV);
+	if (ParsedArgs->OutputFileName == nullptr)
 	{
-		wprintf(L"ERROR: Syntax: cl-filter <dependencies-file> [timing-file] -- <child command line>\n");
-		return -1;
+		wprintf(L"ERROR: No output file name was specified! (%s)\n\n", CmdLine);
+		PrintUsage();
+		return EXIT_CODE_FAILURE;
 	}
 
-	// Get the arguments we care about.
-	bool ParseTimingInfo = wcscmp(ArgV[2], L"--") != 0 && wcscmp(ArgV[3], L"--") == 0;
-	const wchar_t* OutputFileName = ArgV[1];
-	const wchar_t* TimingFileName = ParseTimingInfo ? ArgV[2] : nullptr;
-	const wchar_t* CompilerFileName = ParseTimingInfo ? ArgV[4] : ArgV[3];
+	if (ParsedArgs->CompilerFileName == nullptr)
+	{
+		wprintf(L"ERROR: No compiler file name was found! (%s)\n\n", CmdLine);
+		PrintUsage();
+		return EXIT_CODE_FAILURE;
+	}
 
+	// Get the child command line.
+	wchar_t* ChildCmdLine = wcsstr(CmdLine, L" -- ");
+	if (ChildCmdLine == nullptr)
+	{
+		wprintf(L"ERROR: Unable to find child command line! (%s)\n\n", CmdLine);
+		PrintUsage();
+		return EXIT_CODE_FAILURE;
+	}
+	ChildCmdLine += 4;
+	
 	// Get all the possible localized string prefixes for /showIncludes output
 	std::vector<std::vector<char>> LocalizedIncludePrefixes;
-	GetLocalizedIncludePrefixes(CompilerFileName, LocalizedIncludePrefixes);
+	GetLocalizedIncludePrefixes(ParsedArgs->CompilerFileName, LocalizedIncludePrefixes);
 
 	// Create the child process
 	PROCESS_INFORMATION ProcessInfo;
@@ -56,14 +112,14 @@ int wmain(int ArgC, const wchar_t* ArgV[])
 	if (CreatePipe(&StdOutReadHandle, &StdOutWriteHandle, &SecurityAttributes, 0) == 0)
 	{
 		wprintf(L"ERROR: Unable to create output pipe for child process\n");
-		return -1;
+		return EXIT_CODE_FAILURE;
 	}
 
 	HANDLE StdErrWriteHandle;
 	if (DuplicateHandle(GetCurrentProcess(), StdOutWriteHandle, GetCurrentProcess(), &StdErrWriteHandle, 0, true, DUPLICATE_SAME_ACCESS) == 0)
 	{
 		wprintf(L"ERROR: Unable to create stderr pipe handle for child process\n");
-		return -1;
+		return EXIT_CODE_FAILURE;
 	}
 
 	// Create the new process as suspended, so we can modify it before it starts executing (and potentially preempting us)
@@ -79,7 +135,7 @@ int wmain(int ArgC, const wchar_t* ArgV[])
 	if (CreateProcessW(NULL, ChildCmdLine, NULL, NULL, TRUE, ProcessCreationFlags, NULL, NULL, &StartupInfo, &ProcessInfo) == 0)
 	{
 		wprintf(L"ERROR: Unable to create child process\n");
-		return -1;
+		return EXIT_CODE_FAILURE;
 	}
 
 	// Close the write ends of the handle. We don't want any other process to be able to inherit these.
@@ -87,20 +143,20 @@ int wmain(int ArgC, const wchar_t* ArgV[])
 	CloseHandle(StdErrWriteHandle);
 
 	// Delete the output file.
-	DeleteFileW(OutputFileName);
+	DeleteFileW(ParsedArgs->OutputFileName);
 
 	// Delete the timing file if needed.
-	if (ParseTimingInfo)
+	if (ParsedArgs->TimingFileName != nullptr)
 	{
-		DeleteFileW(TimingFileName);
+		DeleteFileW(ParsedArgs->TimingFileName);
 	}
 
 	// Get the path to a temporary output filename.
-	std::wstring TempOutputFileName(OutputFileName);
+	std::wstring TempOutputFileName(ParsedArgs->OutputFileName);
 	TempOutputFileName += L".tmp";
 
 	// Get the path to a temporary timing filename.
-	std::wstring TempTimingFileName(TimingFileName == nullptr ? L"" : TimingFileName);
+	std::wstring TempTimingFileName(ParsedArgs->TimingFileName == nullptr ? L"" : ParsedArgs->TimingFileName);
 	TempTimingFileName += L".tmp";
 
 	// Create a file to contain the dependency list.
@@ -108,18 +164,18 @@ int wmain(int ArgC, const wchar_t* ArgV[])
 	if(OutputFile == INVALID_HANDLE_VALUE)
 	{
 		wprintf(L"ERROR: Unable to open %s for output", TempOutputFileName.c_str());
-		return -1;
+		return EXIT_CODE_FAILURE;
 	}
 
 	// If needed, create a file to contain unparsed timing info.
 	HANDLE TimingFile = INVALID_HANDLE_VALUE;
-	if (ParseTimingInfo)
+	if (ParsedArgs->TimingFileName != nullptr)
 	{
 		TimingFile = CreateFileW(TempTimingFileName.c_str(), GENERIC_WRITE, FILE_SHARE_READ, NULL, CREATE_ALWAYS, FILE_ATTRIBUTE_NORMAL, NULL);
 		if (TimingFile == INVALID_HANDLE_VALUE)
 		{
 			wprintf(L"ERROR: Unable to open %s for output", TempTimingFileName.c_str());
-			return -1;
+			return EXIT_CODE_FAILURE;
 		}
 	}
 
@@ -191,39 +247,24 @@ int wmain(int ArgC, const wchar_t* ArgV[])
 			}
 
 			// Handling timing info parsing if needed.
-			if (ParseTimingInfo && LineStart < LineEnd)
+			if (ParsedArgs->TimingFileName != nullptr && LineStart < LineEnd)
 			{
 				// See if we've moved into a block of timing information. Order is always Headers -> Classes -> Functions, so to keep for loops to a minimum,
 				//   handle by what our current state is.
 				const char StartTimingInfoText[] = "Include Headers:";
-				const char EndOfTimingInfoText[] = "\tLeast Hits:";
-				if (!InTimingInfo)
-				{
-					if (memcmp(Buffer + LineStart, StartTimingInfoText, 16) == 0)
-					{
-						InTimingInfo = true;
-						TimingLinesToCapture = 1;
-					}
-				}
-				else
-				{
-					if (memcmp(Buffer + LineStart, EndOfTimingInfoText, 12) == 0)
-					{
-						InTimingInfo = false;
-						TimingLinesToCapture = 3;
-					}
-					else
-					{
-						TimingLinesToCapture = 1;
-					}
-				}
-
-				if (TimingLinesToCapture > 0)
+				
+				if (InTimingInfo)
 				{
-					--TimingLinesToCapture;
 					DWORD BytesWritten;
 					WriteFile(TimingFile, Buffer + LineStart, (DWORD)(LineEnd - LineStart), &BytesWritten, NULL);
 					LineStart = LineEnd;
+					continue;
+				}
+				else if (memcmp(Buffer + LineStart, StartTimingInfoText, 16) == 0)
+				{
+					InTimingInfo = true;
+					LineStart = LineEnd;
+					continue;
 				}
 			}
 
@@ -248,23 +289,23 @@ int wmain(int ArgC, const wchar_t* ArgV[])
 	DWORD ExitCode;
 	if (!GetExitCodeProcess(ProcessInfo.hProcess, &ExitCode))
 	{
-		ExitCode = (DWORD)-1;
+		ExitCode = (DWORD)EXIT_CODE_FAILURE;
 	}
 
 	CloseHandle(OutputFile);
 
-	if (ExitCode == 0 && !MoveFileW(TempOutputFileName.c_str(), OutputFileName))
+	if (ExitCode == 0 && !MoveFileW(TempOutputFileName.c_str(), ParsedArgs->OutputFileName))
 	{
-		wprintf(L"ERROR: Unable to rename %s to %s\n", TempOutputFileName.c_str(), OutputFileName);
+		wprintf(L"ERROR: Unable to rename %s to %s\n", TempOutputFileName.c_str(), ParsedArgs->OutputFileName);
 		ExitCode = 1;
 	}
 
-	if (ParseTimingInfo)
+	if (ParsedArgs->TimingFileName != nullptr)
 	{
 		CloseHandle(TimingFile);
-		if (ExitCode == 0 && !MoveFileW(TempTimingFileName.c_str(), TimingFileName))
+		if (ExitCode == 0 && !MoveFileW(TempTimingFileName.c_str(), ParsedArgs->TimingFileName))
 		{
-			wprintf(L"ERROR: Unable to rename %s to %s\n", TempTimingFileName.c_str(), TimingFileName);
+			wprintf(L"ERROR: Unable to rename %s to %s\n", TempTimingFileName.c_str(), ParsedArgs->TimingFileName);
 			ExitCode = 1;
 		}
 	}
diff --git a/Engine/Source/Programs/UnrealBuildTool/Platform/Windows/VCToolChain.cs b/Engine/Source/Programs/UnrealBuildTool/Platform/Windows/VCToolChain.cs
--- a/Engine/Source/Programs/UnrealBuildTool/Platform/Windows/VCToolChain.cs
+++ b/Engine/Source/Programs/UnrealBuildTool/Platform/Windows/VCToolChain.cs
@@ -1145,6 +1145,7 @@ namespace UnrealBuildTool
 				}
 
 				// When compiling with SN-DBS, modules that contain a #import must be built locally
+				CompileAction.bCanExecuteRemotelyWithSNDBS = CompileAction.bCanExecuteRemotely;
 				if (CompileEnvironment.bBuildLocallyWithSNDBS == true)
 				{
 					CompileAction.bCanExecuteRemotelyWithSNDBS = false;
@@ -1156,45 +1157,51 @@ namespace UnrealBuildTool
 			return Result;
 		}
 
+		private Action GenerateParseTimingInfoAction(FileItem SourceFile, Action CompileAction, CPPOutput Result)
+		{
+			FileItem TimingJsonFile = FileItem.GetItemByPath(Path.ChangeExtension(CompileAction.TimingFile.AbsolutePath, ".timing.bin"));
+			Result.DebugDataFiles.Add(TimingJsonFile);
+			Result.DebugDataFiles.Add(CompileAction.TimingFile);
+
+			string ParseTimingArguments = String.Format("-TimingFile=\"{0}\"", CompileAction.TimingFile);
+			if (Target.bParseTimingInfoForTracing)
+			{
+				ParseTimingArguments += " -Tracing";
+			}
+
+			Action ParseTimingInfoAction = Action.CreateRecursiveAction<ParseMsvcTimingInfoMode>(ActionType.ParseTimingInfo, ParseTimingArguments);
+			ParseTimingInfoAction.WorkingDirectory = UnrealBuildTool.EngineSourceDirectory;
+			ParseTimingInfoAction.StatusDescription = String.Format("Parsing Timing File '{0}'", CompileAction.TimingFile);
+			ParseTimingInfoAction.bCanExecuteRemotely = true;
+			ParseTimingInfoAction.bCanExecuteRemotelyWithSNDBS = true;
+			ParseTimingInfoAction.PrerequisiteItems.Add(SourceFile);
+			ParseTimingInfoAction.PrerequisiteItems.Add(CompileAction.TimingFile);
+			ParseTimingInfoAction.ProducedItems.Add(TimingJsonFile);
+			return ParseTimingInfoAction;
+		}
+
 		private void GenerateDependenciesFile(ReadOnlyTargetRules Target, DirectoryReference OutputDir, CPPOutput Result, FileItem SourceFile, List<Action> Actions, Action CompileAction)
 		{
-			List<string> commandArguments = new List<string>();
+			List<string> CommandArguments = new List<string>();
 			CompileAction.DependencyListFile = FileItem.GetItemByFileReference(FileReference.Combine(OutputDir, String.Format("{0}.txt", SourceFile.Location.GetFileName())));
 			CompileAction.ProducedItems.Add(CompileAction.DependencyListFile);
-			commandArguments.Add(Utils.MakePathSafeToUseWithCommandLine(CompileAction.DependencyListFile.Location));
+			CommandArguments.Add(String.Format("-dependencies={0}", Utils.MakePathSafeToUseWithCommandLine(CompileAction.DependencyListFile.Location)));
 
 			if (Target.bPrintToolChainTimingInfo)
 			{
 				CompileAction.TimingFile = FileItem.GetItemByFileReference(FileReference.Combine(OutputDir, String.Format("{0}.timing.txt", SourceFile.Location.GetFileName())));
 				CompileAction.ProducedItems.Add(CompileAction.TimingFile);
-				commandArguments.Add(Utils.MakePathSafeToUseWithCommandLine(CompileAction.TimingFile.Location));
-
-				FileItem TimingJsonFile = FileItem.GetItemByPath(Path.ChangeExtension(CompileAction.TimingFile.AbsolutePath, ".timing.bin"));
-				Result.DebugDataFiles.Add(TimingJsonFile);
-				Result.DebugDataFiles.Add(CompileAction.TimingFile);
-
-				string ParseTimingArguments = String.Format("-TimingFile=\"{0}\"", CompileAction.TimingFile);
-				if (Target.bParseTimingInfoForTracing)
-				{
-					ParseTimingArguments += " -Tracing";
-				}
-
-				Action ParseTimingInfoAction = Action.CreateRecursiveAction<ParseMsvcTimingInfoMode>(ActionType.ParseTimingInfo, ParseTimingArguments);
-				ParseTimingInfoAction.WorkingDirectory = UnrealBuildTool.EngineSourceDirectory;
-				ParseTimingInfoAction.StatusDescription = String.Format("Parsing Timing File '{0}'", CompileAction.TimingFile);
-				ParseTimingInfoAction.bCanExecuteRemotely = true;
-				ParseTimingInfoAction.bCanExecuteRemotelyWithSNDBS = true;
-				ParseTimingInfoAction.PrerequisiteItems.Add(SourceFile);
-				ParseTimingInfoAction.PrerequisiteItems.Add(CompileAction.TimingFile);
-				ParseTimingInfoAction.ProducedItems.Add(TimingJsonFile);
+				Action ParseTimingInfoAction = GenerateParseTimingInfoAction(SourceFile, CompileAction, Result);
 				Actions.Add(ParseTimingInfoAction);
+				CommandArguments.Add(string.Format("-timing={0}", Utils.MakePathSafeToUseWithCommandLine(CompileAction.TimingFile.Location)));
 			}
 
-			commandArguments.Add("--");
-			commandArguments.Add(Utils.MakePathSafeToUseWithCommandLine(CompileAction.CommandPath));
-			commandArguments.Add(CompileAction.CommandArguments);
-			commandArguments.Add("/showIncludes");
-			CompileAction.CommandArguments = string.Join(" ", commandArguments);
+			CommandArguments.Add(String.Format("-compiler={0}", Utils.MakePathSafeToUseWithCommandLine(CompileAction.CommandPath)));
+			CommandArguments.Add("--");
+			CommandArguments.Add(Utils.MakePathSafeToUseWithCommandLine(CompileAction.CommandPath));
+			CommandArguments.Add(CompileAction.CommandArguments);
+			CommandArguments.Add("/showIncludes");
+			CompileAction.CommandArguments = string.Join(" ", CommandArguments);
 			CompileAction.CommandPath = FileReference.Combine(UnrealBuildTool.EngineDirectory, "Build", "Windows", "cl-filter", "cl-filter.exe");	
 		}
 
@@ -1217,6 +1224,7 @@ namespace UnrealBuildTool
 					AggregateTimingInfoAction.WorkingDirectory = UnrealBuildTool.EngineSourceDirectory;
 					AggregateTimingInfoAction.StatusDescription = $"Aggregating {TimingJsonFiles.Count} Timing File(s)";
 					AggregateTimingInfoAction.bCanExecuteRemotely = false;
+					AggregateTimingInfoAction.bCanExecuteRemotelyWithSNDBS = false;
 					AggregateTimingInfoAction.PrerequisiteItems.AddRange(ParseTimingActions.SelectMany(a => a.PrerequisiteItems));
 					AggregateTimingInfoAction.PrerequisiteItems.AddRange(TimingJsonFiles);
 
