UBT: Fix an issue where cl-filter would not correctly parse full lines of data, preventing included files from being parsed correctly.

diff --git a/Engine/Extras/cl-filter/cl-filter.cpp b/Engine/Extras/cl-filter/cl-filter.cpp
--- a/Engine/Extras/cl-filter/cl-filter.cpp
+++ b/Engine/Extras/cl-filter/cl-filter.cpp
@@ -138,6 +138,9 @@ int wmain(int ArgC, const wchar_t* ArgV[])
 		return EXIT_CODE_FAILURE;
 	}
 
+	// Close the startup thread handle; we don't need it.
+	CloseHandle(ProcessInfo.hThread);
+
 	// Close the write ends of the handle. We don't want any other process to be able to inherit these.
 	CloseHandle(StdOutWriteHandle);
 	CloseHandle(StdErrWriteHandle);
@@ -188,9 +191,9 @@ int wmain(int ArgC, const wchar_t* ArgV[])
 	for (;;)
 	{
 		// Read the next chunk of data from the output stream
+		DWORD BytesRead = 0;
 		if (BufferSize < sizeof(Buffer))
 		{
-			DWORD BytesRead;
 			if (ReadFile(StdOutReadHandle, Buffer + BufferSize, (DWORD)(sizeof(Buffer) - BufferSize), &BytesRead, NULL))
 			{
 				BufferSize += BytesRead;
@@ -217,7 +220,7 @@ int wmain(int ArgC, const wchar_t* ArgV[])
 			}
 
 			// If we didn't reach a line terminator, and we can still read more data, clear up some space and try again
-			if (LineEnd == BufferSize && LineStart != 0)
+			if (LineEnd == BufferSize && !(LineStart == 0 && BytesRead == 0) && !(LineStart == 0 && BufferSize == sizeof(Buffer)))
 			{
 				break;
 			}
