Change how cl-filter buffers output text. Since MultiByteToWideChar doesn't report back how many input characters were actually used, we can drop output that wasn't converted into the output buffer if the stdout pipe is full. Instead, always look for ANSI markers indicating the presence of an include file, and redirect based on that. This makes the assumption that CL.EXE outputs single-byte ANSI strings for all supported codepages; if not, we'll either have to buffer the whole output and convert at once, or step through the output text one character at a time.

#rnx

diff --git a/Engine/Extras/cl-filter/cl-filter.cpp b/Engine/Extras/cl-filter/cl-filter.cpp
--- a/Engine/Extras/cl-filter/cl-filter.cpp
+++ b/Engine/Extras/cl-filter/cl-filter.cpp
@@ -87,10 +87,10 @@ int wmain(int ArgC, const wchar_t* ArgV[])
 	TempOutputFileName += L".tmp";
 
 	// Create a file to contain the dependency list
-	FILE* OutputFile;
-	if (_wfopen_s(&OutputFile, TempOutputFileName.c_str(), L"wt") != 0)
+	HANDLE OutputFile = CreateFileW(TempOutputFileName.c_str(), GENERIC_WRITE, FILE_SHARE_READ, NULL, CREATE_ALWAYS, FILE_ATTRIBUTE_NORMAL, NULL);
+	if(OutputFile == INVALID_HANDLE_VALUE)
 	{
-		wprintf(L"ERROR: Unable to open %s for output", OutputFileName);
+		wprintf(L"ERROR: Unable to open %s for output", TempOutputFileName.c_str());
 		return -1;
 	}
 
@@ -98,31 +98,26 @@ int wmain(int ArgC, const wchar_t* ArgV[])
 	UINT CodePage = GetConsoleCP();
 	
 	// Pipe the output to stdout
-	wchar_t Buffer[4096];
+	char Buffer[1024];
 	size_t BufferSize = 0;
 
 	for (;;)
 	{
-		char InputData[4096];
-
 		// Read the next chunk of data from the output stream
-		DWORD BytesRead;
-		if (!ReadFile(StdOutReadHandle, InputData, sizeof(InputData), &BytesRead, NULL))
+		if (BufferSize < sizeof(Buffer))
 		{
-			if (GetLastError() != ERROR_BROKEN_PIPE)
+			DWORD BytesRead;
+			if (ReadFile(StdOutReadHandle, Buffer, (DWORD)(sizeof(Buffer) - BufferSize), &BytesRead, NULL))
+			{
+				BufferSize += BytesRead;
+			}
+			else if(GetLastError() != ERROR_BROKEN_PIPE)
 			{
 				wprintf(L"ERROR: Unable to read data from child process (%08x)", GetLastError());
 			}
-			BytesRead = 0;
-		}
-
-		// Convert it to a string
-		if (BytesRead > 0)
-		{
-			int NumCharacters = MultiByteToWideChar(CodePage, 0, InputData, (int)BytesRead, Buffer + BufferSize, (int)((sizeof(Buffer) / sizeof(Buffer[0])) - BufferSize));
-			if (NumCharacters > 0)
+			else if (BufferSize == 0)
 			{
-				BufferSize += (size_t)NumCharacters;
+				break;
 			}
 		}
 
@@ -153,46 +148,31 @@ int wmain(int ArgC, const wchar_t* ArgV[])
 				LineEnd++;
 			}
 
-			// Filter out any lines that are 
-			const wchar_t Prefix[] = L"Note: including file: ";
-			if (wcsncmp(Buffer + LineStart, Prefix, sizeof(Prefix) / sizeof(Prefix[0]) - 1) == 0)
+			// Filter out any lines that have the "Note: including file: " prefix.
+			const char Prefix[] = "Note: including file: ";
+			if (strncmp(Buffer + LineStart, Prefix, sizeof(Prefix) / sizeof(Prefix[0]) - 1) == 0)
 			{
 				size_t FileNameStart = LineStart + sizeof(Prefix) / sizeof(Prefix[0]) - 1;
-				while (FileNameStart < LineEnd && iswspace(Buffer[FileNameStart]))
+				while (FileNameStart < LineEnd && isspace(Buffer[FileNameStart]))
 				{
 					FileNameStart++;
 				}
 
-				size_t FileNameEnd = LineEnd;
-				while (FileNameEnd > FileNameStart && iswspace(Buffer[FileNameEnd - 1]))
-				{
-					FileNameEnd--;
-				}
-
-				Buffer[FileNameEnd] = 0;
-				fwprintf(OutputFile, L"%s\n", Buffer + FileNameStart);
+				DWORD BytesWritten;
+				WriteFile(OutputFile, Buffer + FileNameStart, (DWORD)(LineEnd - FileNameStart), &BytesWritten, NULL);
 			}
 			else
 			{
-				char RawOutput[4096];
-				int RawOutputSize = WideCharToMultiByte(CodePage, 0, Buffer + LineStart, (int)(LineEnd - LineStart), RawOutput, sizeof(RawOutput), nullptr, nullptr);
-
 				DWORD BytesWritten;
-				WriteFile(GetStdHandle(STD_OUTPUT_HANDLE), RawOutput, RawOutputSize, &BytesWritten, NULL);
+				WriteFile(GetStdHandle(STD_OUTPUT_HANDLE), Buffer + LineStart, (DWORD)(LineEnd - LineStart), &BytesWritten, NULL);
 			}
 
 			// Move to the next line
 			LineStart = LineEnd;
 		}
 
-		// Check if we've finished
-		if (BufferSize == 0 && BytesRead == 0)
-		{
-			break;
-		}
-
 		// Shuffle everything down
-		memmove(Buffer, Buffer + LineStart, (BufferSize - LineStart) * sizeof(wchar_t));
+		memmove(Buffer, Buffer + LineStart, BufferSize - LineStart);
 		BufferSize -= LineStart;
 	}
 
@@ -204,9 +184,9 @@ int wmain(int ArgC, const wchar_t* ArgV[])
 		ExitCode = (DWORD)-1;
 	}
 
-	fclose(OutputFile);
+	CloseHandle(OutputFile);
 
-	if (ExitCode == 0 && !MoveFile(TempOutputFileName.c_str(), OutputFileName))
+	if (ExitCode == 0 && !MoveFileW(TempOutputFileName.c_str(), OutputFileName))
 	{
 		wprintf(L"ERROR: Unable to rename %s to %s\n", TempOutputFileName.c_str(), OutputFileName);
 		ExitCode = 1;
