Fix cl-filter always reading into the start of the buffer.

#rnx

diff --git a/Engine/Extras/cl-filter/cl-filter.cpp b/Engine/Extras/cl-filter/cl-filter.cpp
--- a/Engine/Extras/cl-filter/cl-filter.cpp
+++ b/Engine/Extras/cl-filter/cl-filter.cpp
@@ -107,7 +107,7 @@ int wmain(int ArgC, const wchar_t* ArgV[])
 		if (BufferSize < sizeof(Buffer))
 		{
 			DWORD BytesRead;
-			if (ReadFile(StdOutReadHandle, Buffer, (DWORD)(sizeof(Buffer) - BufferSize), &BytesRead, NULL))
+			if (ReadFile(StdOutReadHandle, Buffer + BufferSize, (DWORD)(sizeof(Buffer) - BufferSize), &BytesRead, NULL))
 			{
 				BufferSize += BytesRead;
 			}
