Change cl-filter to not treat carriage return characters as line terminators, so the input buffer isn't split on them if a following carriage return didn't fit.

#rnx

diff --git a/Engine/Extras/cl-filter/cl-filter.cpp b/Engine/Extras/cl-filter/cl-filter.cpp
--- a/Engine/Extras/cl-filter/cl-filter.cpp
+++ b/Engine/Extras/cl-filter/cl-filter.cpp
@@ -127,7 +127,7 @@ int wmain(int ArgC, const wchar_t* ArgV[])
 		{
 			// Find the end of this line
 			size_t LineEnd = LineStart;
-			while (LineEnd < BufferSize && Buffer[LineEnd] != '\r' && Buffer[LineEnd] != '\n')
+			while (LineEnd < BufferSize && Buffer[LineEnd] != '\n')
 			{
 				LineEnd++;
 			}
@@ -138,11 +138,7 @@ int wmain(int ArgC, const wchar_t* ArgV[])
 				break;
 			}
 
-			// Skip past the EOL markers
-			if (LineEnd < BufferSize && Buffer[LineEnd] == '\r')
-			{
-				LineEnd++;
-			}
+			// Skip past the EOL marker
 			if (LineEnd < BufferSize && Buffer[LineEnd] == '\n')
 			{
 				LineEnd++;
