Fix scraping of include paths from compiler if a non-english language pack is enabled. Take the system locale, user locale and VSLANG environment variable, and check for localized versions of the "Note: Including file" message.

diff --git a/Engine/Extras/cl-filter/cl-filter.cpp b/Engine/Extras/cl-filter/cl-filter.cpp
--- a/Engine/Extras/cl-filter/cl-filter.cpp
+++ b/Engine/Extras/cl-filter/cl-filter.cpp
@@ -9,6 +9,9 @@
 #include <stdio.h>
 #include <string>
 #include <vector>
+#include <set>
+
+void GetLocalizedIncludePrefixes(const wchar_t* CompilerPath, std::vector<std::vector<char>>& Prefixes);
 
 int wmain(int ArgC, const wchar_t* ArgV[])
 {
@@ -21,20 +24,22 @@ int wmain(int ArgC, const wchar_t* ArgV[])
 		wprintf(L"ERROR: Unable to find child command line (%s)", CmdLine);
 		return -1;
 	}
-
-	*ChildCmdLine = 0;
 	ChildCmdLine += 4;
 
-	int NumArgs = 0;
-	wchar_t** Args = CommandLineToArgvW(CmdLine, &NumArgs);
-
-	if (NumArgs != 2)
+	// Make sure we've got an output file and compiler path
+	if (ArgC <= 4 || wcscmp(ArgV[2], L"--") != 0)
 	{
 		wprintf(L"ERROR: Syntax: cl-filter <dependencies-file> -- <child command line>\n");
 		return -1;
 	}
 
-	const wchar_t* OutputFileName = Args[1];
+	// Get the arguments we care about
+	const wchar_t* OutputFileName = ArgV[1];
+	const wchar_t* CompilerFileName = ArgV[3];
+
+	// Get all the possible localized string prefixes for /showIncludes output
+	std::vector<std::vector<char>> LocalizedPrefixes;
+	GetLocalizedIncludePrefixes(CompilerFileName, LocalizedPrefixes);
 
 	// Create the child process
 	PROCESS_INFORMATION ProcessInfo;
@@ -145,19 +150,25 @@ int wmain(int ArgC, const wchar_t* ArgV[])
 			}
 
 			// Filter out any lines that have the "Note: including file: " prefix.
-			const char Prefix[] = "Note: including file: ";
-			if (strncmp(Buffer + LineStart, Prefix, sizeof(Prefix) / sizeof(Prefix[0]) - 1) == 0)
+			for (const std::vector<char>& LocalizedPrefix : LocalizedPrefixes)
 			{
-				size_t FileNameStart = LineStart + sizeof(Prefix) / sizeof(Prefix[0]) - 1;
-				while (FileNameStart < LineEnd && isspace(Buffer[FileNameStart]))
+				if (memcmp(Buffer + LineStart, LocalizedPrefix.data(), LocalizedPrefix.size() - 1) == 0)
 				{
-					FileNameStart++;
+					size_t FileNameStart = LineStart + LocalizedPrefix.size() - 1;
+					while (FileNameStart < LineEnd && isspace(Buffer[FileNameStart]))
+					{
+						FileNameStart++;
+					}
+
+					DWORD BytesWritten;
+					WriteFile(OutputFile, Buffer + FileNameStart, (DWORD)(LineEnd - FileNameStart), &BytesWritten, NULL);
+					LineStart = LineEnd;
+					break;
 				}
-
-				DWORD BytesWritten;
-				WriteFile(OutputFile, Buffer + FileNameStart, (DWORD)(LineEnd - FileNameStart), &BytesWritten, NULL);
 			}
-			else
+
+			// If we didn't write anything out, write it to stdout
+			if(LineStart < LineEnd)
 			{
 				DWORD BytesWritten;
 				WriteFile(GetStdHandle(STD_OUTPUT_HANDLE), Buffer + LineStart, (DWORD)(LineEnd - LineStart), &BytesWritten, NULL);
@@ -190,3 +201,127 @@ int wmain(int ArgC, const wchar_t* ArgV[])
 
 	return ExitCode;
 }
+
+static std::string FindAndReplace(std::string Input, const std::string& FindStr, const std::string& ReplaceStr)
+{
+	size_t Start = 0;
+	for (;;)
+	{
+		size_t Offset = Input.find(FindStr, Start);
+		if(Offset == std::wstring::npos)
+		{
+			break;
+		}
+
+		Input.replace(Offset, FindStr.size(), ReplaceStr);
+		Start = Offset + ReplaceStr.size();
+	}
+	return Input;
+}
+
+bool GetLocalizedIncludePrefix(const wchar_t* LibraryPath, HMODULE LibraryHandle, std::vector<char>& Prefix)
+{
+	static const unsigned int ResourceId = 408;
+
+	// Read the string from the library
+	wchar_t Text[512];
+	if(LoadStringW(LibraryHandle, ResourceId, Text, sizeof(Text) / sizeof(Text[0])) == 0)
+	{
+		wprintf(L"WARNING: unable to read string %d from %s\n", ResourceId, LibraryPath);
+		return false;
+	}
+
+	// Find the end of the prefix 
+	wchar_t* TextEnd = wcsstr(Text, L"%s%s");
+	if (TextEnd == nullptr)
+	{
+		wprintf(L"WARNING: unable to find substitution markers in format string '%s' (%s)", Text, LibraryPath);
+		return false;
+	}
+
+	// Figure out how large the buffer needs to be to hold the MBCS version
+	int Length = WideCharToMultiByte(CP_ACP, 0, Text, (int)(TextEnd - Text), NULL, 0, NULL, NULL);
+	if (Length == 0)
+	{
+		wprintf(L"WARNING: unable to query size for MBCS output buffer (input text '%s', library %s)", Text, LibraryPath);
+		return false;
+	}
+
+	// Resize the output buffer with space for a null terminator
+	Prefix.resize(Length + 1);
+
+	// Get the multibyte text
+	int Result = WideCharToMultiByte(CP_ACP, 0, Text, (int)(TextEnd - Text), Prefix.data(), Length, NULL, NULL);
+	if (Result == 0)
+	{
+		wprintf(L"WARNING: unable to get MBCS string (input text '%s', library %s)", Text, LibraryPath);
+		return false;
+	}
+
+	return true;
+}
+
+// Language packs for Visual Studio contain localized strings for the "Note: including file:" prefix we expect to see when running the compiler
+// with the /showIncludes option. Enumerate all the possible languages that may be active, and build up an array of possible prefixes. We'll treat
+// any of them as markers for included files.
+void GetLocalizedIncludePrefixes(const wchar_t* CompilerPath, std::vector<std::vector<char>>& Prefixes)
+{
+	// Get all the possible locale id's. Include en-us by default.
+	std::set<std::wstring> LocaleIds;
+	LocaleIds.insert(L"1033");
+
+	// The user default locale id
+	wchar_t LocaleIdString[20];
+	wsprintf(LocaleIdString, L"%d", GetUserDefaultLCID());
+	LocaleIds.insert(LocaleIdString);
+
+	// The system default locale id
+	wsprintf(LocaleIdString, L"%d", GetSystemDefaultLCID());
+	LocaleIds.insert(LocaleIdString);
+
+	// The Visual Studio locale setting
+	static const size_t VsLangMaxLen = 256;
+	wchar_t VsLangEnv[VsLangMaxLen];
+	if (GetEnvironmentVariableW(L"VSLANG", VsLangEnv, VsLangMaxLen) != 0)
+	{
+		LocaleIds.insert(VsLangEnv);
+	}
+
+	// Find the directory containing the compiler path
+	size_t CompilerDirLen = wcslen(CompilerPath);
+	while (CompilerDirLen > 0 && CompilerPath[CompilerDirLen - 1] != '/' && CompilerPath[CompilerDirLen - 1] != '\\')
+	{
+		CompilerDirLen--;
+	}
+
+	// Always add the en-us prefix. We'll validate that this is correct if we have an en-us resource file, but it gives us something to fall back on.
+	const char EnglishText[] = "Note: including file:";
+	Prefixes.emplace_back(EnglishText, strchr(EnglishText, 0) + 1);
+	
+	// Loop through all the possible locale ids and try to find the localized string for each
+	for (const std::wstring& LocaleId : LocaleIds)
+	{
+		std::wstring ResourceFile;
+		ResourceFile.assign(CompilerPath, CompilerPath + CompilerDirLen);
+		ResourceFile.append(LocaleId);
+		ResourceFile.append(L"\\clui.dll");
+
+		HMODULE LibraryHandle = LoadLibraryExW(ResourceFile.c_str(), 0, LOAD_LIBRARY_AS_DATAFILE | LOAD_LIBRARY_AS_IMAGE_RESOURCE);
+		if (LibraryHandle != nullptr)
+		{
+			std::vector<char> Prefix;
+			if (GetLocalizedIncludePrefix(ResourceFile.c_str(), LibraryHandle, Prefix))
+			{
+				if (wcscmp(LocaleId.c_str(), L"1033") != 0)
+				{
+					Prefixes.push_back(std::move(Prefix));
+				}
+				else if(strcmp(Prefix.data(), EnglishText) != 0)
+				{
+					wprintf(L"WARNING: unexpected localized string for en-us.\n   Expected: '%hs'\n   Actual:   '%hs'", FindAndReplace(EnglishText, "\n", "\\n").c_str(), FindAndReplace(Prefix.data(), "\n", "\\n").c_str());
+				}
+			}
+			FreeLibrary(LibraryHandle);
+		}
+	}
+}
