CL-Filter: Fix incorrect codepage being used to convert unicode strings into expected output.

diff --git a/Engine/Extras/cl-filter/cl-filter.cpp b/Engine/Extras/cl-filter/cl-filter.cpp
--- a/Engine/Extras/cl-filter/cl-filter.cpp
+++ b/Engine/Extras/cl-filter/cl-filter.cpp
@@ -99,9 +99,6 @@ int wmain(int ArgC, const wchar_t* ArgV[])
 		return -1;
 	}
 
-	// Get the default console codepage
-	UINT CodePage = GetConsoleCP();
-	
 	// Pipe the output to stdout
 	char Buffer[1024];
 	size_t BufferSize = 0;
@@ -219,7 +216,7 @@ static std::string FindAndReplace(std::string Input, const std::string& FindStr,
 	return Input;
 }
 
-bool GetLocalizedIncludePrefix(const wchar_t* LibraryPath, HMODULE LibraryHandle, std::vector<char>& Prefix)
+bool GetLocalizedIncludePrefix(UINT CodePage, const wchar_t* LibraryPath, HMODULE LibraryHandle, std::vector<char>& Prefix)
 {
 	static const unsigned int ResourceId = 408;
 
@@ -251,7 +248,7 @@ bool GetLocalizedIncludePrefix(const wchar_t* LibraryPath, HMODULE LibraryHandle
 	Prefix.resize(Length + 1);
 
 	// Get the multibyte text
-	int Result = WideCharToMultiByte(CP_ACP, 0, Text, (int)(TextEnd - Text), Prefix.data(), Length, NULL, NULL);
+	int Result = WideCharToMultiByte(CodePage, 0, Text, (int)(TextEnd - Text), Prefix.data(), Length, NULL, NULL);
 	if (Result == 0)
 	{
 		wprintf(L"WARNING: unable to get MBCS string (input text '%s', library %s)", Text, LibraryPath);
@@ -298,6 +295,9 @@ void GetLocalizedIncludePrefixes(const wchar_t* CompilerPath, std::vector<std::v
 	const char EnglishText[] = "Note: including file:";
 	Prefixes.emplace_back(EnglishText, strchr(EnglishText, 0) + 1);
 	
+	// Get the default console codepage
+	UINT CodePage = GetConsoleOutputCP();
+
 	// Loop through all the possible locale ids and try to find the localized string for each
 	for (const std::wstring& LocaleId : LocaleIds)
 	{
@@ -310,7 +310,7 @@ void GetLocalizedIncludePrefixes(const wchar_t* CompilerPath, std::vector<std::v
 		if (LibraryHandle != nullptr)
 		{
 			std::vector<char> Prefix;
-			if (GetLocalizedIncludePrefix(ResourceFile.c_str(), LibraryHandle, Prefix))
+			if (GetLocalizedIncludePrefix(CodePage, ResourceFile.c_str(), LibraryHandle, Prefix))
 			{
 				if (wcscmp(LocaleId.c_str(), L"1033") != 0)
 				{
